Updating D:\Dev-Build/Engine/Build/ExpectedCompileTimes/UnrealPak.json.

diff --git a/Engine/Build/ExpectedCompileTimes/UnrealPak.json b/Engine/Build/ExpectedCompileTimes/UnrealPak.json
new file mode 100644
--- /dev/null
+++ b/Engine/Build/ExpectedCompileTimes/UnrealPak.json
@@ -0,0 +1,27 @@
+{
+	"D:\\Dev-Build\\Engine\\Intermediate\\Build\\Win64\\UnrealPak\\Development\\TraceLog\\Module.TraceLog.cpp": 0.0228770003084264,
+	"D:\\Dev-Build\\Engine\\Intermediate\\Build\\Win64\\UnrealPak\\Development\\Core\\Module.Core.11_of_12.cpp": 5.02142098795548,
+	"D:\\Dev-Build\\Engine\\Intermediate\\Build\\Win64\\UnrealPak\\Development\\Core\\Module.Core.5_of_12.cpp": 4.21581300567277,
+	"D:\\Dev-Build\\Engine\\Intermediate\\Build\\Win64\\UnrealPak\\Development\\UnrealPak\\UnrealPak.cpp": 5.18721699405694,
+	"D:\\Dev-Build\\Engine\\Intermediate\\Build\\Win64\\UnrealPak\\Development\\Core\\Module.Core.12_of_12.cpp": 2.28579801642582,
+	"D:\\Dev-Build\\Engine\\Intermediate\\Build\\Win64\\UnrealPak\\Development\\Core\\Module.Core.6_of_12.cpp": 1.28478300825782,
+	"D:\\Dev-Build\\Engine\\Intermediate\\Build\\Win64\\UnrealPak\\Development\\Core\\PCH.Core.h": 3.23119196821779,
+	"D:\\Dev-Build\\Engine\\Intermediate\\Build\\Win64\\UnrealPak\\Development\\Core\\SharedPCH.Core.h": 2.60879299279895,
+	"D:\\Dev-Build\\Engine\\Intermediate\\Build\\Win64\\UnrealPak\\Development\\RSA\\RSA.cpp": 0.0573219998259447,
+	"D:\\Dev-Build\\Engine\\Plugins\\Compression\\XboxOne\\XboxOneZlib\\Intermediate\\Build\\Win64\\UnrealPak\\Development\\XboxOneZlibCompressionFormat\\XboxOneZlibCompressionFormat.cpp": 0.0358570001753833,
+	"D:\\Dev-Build\\Engine\\Intermediate\\Build\\Win64\\UnrealPak\\Development\\Core\\Module.Core.1_of_12.cpp": 2.38129599958438,
+	"D:\\Dev-Build\\Engine\\Intermediate\\Build\\Win64\\UnrealPak\\Development\\PakFile\\Module.PakFile.cpp": 1.36390900652123,
+	"D:\\Dev-Build\\Engine\\Plugins\\Experimental\\PlatformCrypto\\Intermediate\\Build\\Win64\\UnrealPak\\Development\\PlatformCrypto\\PlatformCrypto.cpp": 0.0434490000034202,
+	"D:\\Dev-Build\\Engine\\Intermediate\\Build\\Win64\\UnrealPak\\Development\\Core\\Module.Core.7_of_12.cpp": 2.60453498547349,
+	"D:\\Dev-Build\\Engine\\Intermediate\\Build\\Win64\\UnrealPak\\Development\\Core\\Module.Core.2_of_12.cpp": 2.51389199943242,
+	"D:\\Dev-Build\\Engine\\Intermediate\\Build\\Win64\\UnrealPak\\Development\\Json\\Module.Json.cpp": 1.35231898639131,
+	"D:\\Dev-Build\\Engine\\Plugins\\Experimental\\PlatformCrypto\\Intermediate\\Build\\Win64\\UnrealPak\\Development\\PlatformCryptoOpenSSL\\Module.PlatformCryptoOpenSSL.cpp": 0.237803994595197,
+	"D:\\Dev-Build\\Engine\\Intermediate\\Build\\Win64\\UnrealPak\\Development\\Core\\Module.Core.3_of_12.cpp": 2.88178096901777,
+	"D:\\Dev-Build\\Engine\\Intermediate\\Build\\Win64\\UnrealPak\\Development\\BuildSettings\\BuildSettings.cpp": 0.0185460003024218,
+	"D:\\Dev-Build\\Engine\\Intermediate\\Build\\Win64\\UnrealPak\\Development\\Projects\\Module.Projects.cpp": 1.11068999695772,
+	"D:\\Dev-Build\\Engine\\Intermediate\\Build\\Win64\\UnrealPak\\Development\\Core\\Module.Core.8_of_12.cpp": 4.01791000457888,
+	"D:\\Dev-Build\\Engine\\Intermediate\\Build\\Win64\\UnrealPak\\Development\\PakFileUtilities\\Module.PakFileUtilities.cpp": 1.68093500445411,
+	"D:\\Dev-Build\\Engine\\Intermediate\\Build\\Win64\\UnrealPak\\Development\\Core\\Module.Core.4_of_12.cpp": 3.15408397776991,
+	"D:\\Dev-Build\\Engine\\Intermediate\\Build\\Win64\\UnrealPak\\Development\\Core\\Module.Core.9_of_12.cpp": 2.4905550079402,
+	"D:\\Dev-Build\\Engine\\Intermediate\\Build\\Win64\\UnrealPak\\Development\\Core\\Module.Core.10_of_12.cpp": 2.06628699420162
+}
\ No newline at end of file
